Ensure the build workers don't use a DDC backend graph.

diff --git a/Engine/Source/Programs/DerivedDataBuildWorker/Private/DerivedDataBuildWorker.cpp b/Engine/Source/Programs/DerivedDataBuildWorker/Private/DerivedDataBuildWorker.cpp
--- a/Engine/Source/Programs/DerivedDataBuildWorker/Private/DerivedDataBuildWorker.cpp
+++ b/Engine/Source/Programs/DerivedDataBuildWorker/Private/DerivedDataBuildWorker.cpp
@@ -19,7 +19,7 @@ DEFINE_LOG_CATEGORY_STATIC(LogDerivedDataBuildWorker, Log, All);
 INT32_MAIN_INT32_ARGC_TCHAR_ARGV()
 {
 	FTaskTagScope Scope(ETaskTag::EGameThread);
-	FString CmdLine = FCommandLine::BuildFromArgV(nullptr, ArgC, ArgV, nullptr);
+	FString CmdLine = FCommandLine::BuildFromArgV(nullptr, ArgC, ArgV, TEXT("-ddc=None"));
 	GEngineLoop.PreInit(*CmdLine);
 
 	// Make sure the engine is properly cleaned up whenever we exit this function
