Fix DerivedDataBuildLoop.cpp

diff --git a/Engine/Source/Programs/DerivedDataBuildWorker/Private/DerivedDataBuildLoop.cpp b/Engine/Source/Programs/DerivedDataBuildWorker/Private/DerivedDataBuildLoop.cpp
--- a/Engine/Source/Programs/DerivedDataBuildWorker/Private/DerivedDataBuildLoop.cpp
+++ b/Engine/Source/Programs/DerivedDataBuildWorker/Private/DerivedDataBuildLoop.cpp
@@ -96,7 +96,7 @@ public:
 		}
 	}
 
-	virtual FCbObject GetConstant(FStringView Key) const override
+	virtual FCbObject FindConstant(FStringView Key) const override
 	{
 		FCbObject FoundVal;
 		BuildActionRecord.BuildAction.IterateConstants([&Key, &FoundVal] (FStringView ConstantKey, FCbObject&& Value)
@@ -107,7 +107,7 @@ public:
 		return FoundVal;
 	}
 
-	virtual FSharedBuffer GetInput(FStringView Key) const override
+	virtual FSharedBuffer FindInput(FStringView Key) const override
 	{
 		FSharedBuffer FoundVal;
 		BuildActionRecord.BuildAction.IterateInputs([this, &Key, &FoundVal] (FStringView InputKey, const FIoHash& RawHash, uint64 RawSize)
