DDC: Convert the build worker to use the build session and input resolver

#rnx

diff --git a/Engine/Source/Programs/DerivedDataBuildWorker/Private/DerivedDataBuildLoop.cpp b/Engine/Source/Programs/DerivedDataBuildWorker/Private/DerivedDataBuildLoop.cpp
deleted file mode 100644
diff --git a/Engine/Source/Programs/DerivedDataBuildWorker/Private/DerivedDataBuildLoop.h b/Engine/Source/Programs/DerivedDataBuildWorker/Private/DerivedDataBuildLoop.h
deleted file mode 100644
diff --git a/Engine/Source/Programs/DerivedDataBuildWorker/Private/DerivedDataBuildWorker.cpp b/Engine/Source/Programs/DerivedDataBuildWorker/Private/DerivedDataBuildWorker.cpp
--- a/Engine/Source/Programs/DerivedDataBuildWorker/Private/DerivedDataBuildWorker.cpp
+++ b/Engine/Source/Programs/DerivedDataBuildWorker/Private/DerivedDataBuildWorker.cpp
@@ -1,64 +1,333 @@
 // Copyright Epic Games, Inc. All Rights Reserved.
 
-#include "CoreMinimal.h"
-
-#include "Containers/Array.h"
-#include "CoreTypes.h"
-#include "DerivedDataBuildFunction.h"
-#include "DerivedDataBuildFunctionFactory.h"
-#include "DerivedDataBuildLoop.h"
-#include "Features/IModularFeatures.h"
-#include "Logging/LogMacros.h"
+#include "Compression/CompressedBuffer.h"
+#include "Containers/UnrealString.h"
+#include "DerivedDataBuild.h"
+#include "DerivedDataBuildAction.h"
+#include "DerivedDataBuildDefinition.h"
+#include "DerivedDataBuildInputs.h"
+#include "DerivedDataBuildInputResolver.h"
+#include "DerivedDataBuildOutput.h"
+#include "DerivedDataBuildSession.h"
+#include "DerivedDataPayload.h"
+#include "HAL/FileManager.h"
+#include "Memory/SharedBuffer.h"
+#include "Misc/CommandLine.h"
+#include "Misc/CoreMisc.h"
+#include "Misc/Parse.h"
+#include "Misc/Paths.h"
+#include "Misc/PathViews.h"
 #include "Misc/ScopeExit.h"
+#include "Misc/WildcardString.h"
+#include "Serialization/CompactBinary.h"
+#include "Serialization/CompactBinarySerialization.h"
+#include "Serialization/CompactBinaryWriter.h"
+#include "String/ParseTokens.h"
+
 #include "RequiredProgramMainCPPInclude.h"
 
 DEFINE_LOG_CATEGORY_STATIC(LogDerivedDataBuildWorker, Log, All);
 
-//////////////////////////////////////////////////////////////////////////
+namespace UE::DerivedData
+{
 
-INT32_MAIN_INT32_ARGC_TCHAR_ARGV()
+class FBuildWorkerProgram : public IBuildInputResolver
+{
+public:
+	bool ParseCommandLine(const TCHAR* CommandLine);
+	bool Build();
+
+private:
+	void BuildComplete(FBuildActionCompleteParams&& Params) const;
+
+	FRequest ResolveInputData(const FBuildAction& Action, EPriority Priority, FOnBuildInputDataResolved&& OnResolved, FBuildInputFilter&& Filter) final;
+
+	TUniquePtr<FArchive> OpenInput(FStringView ActionPath, const FIoHash& RawHash) const;
+	TUniquePtr<FArchive> OpenOutput(FStringView ActionPath, const FIoHash& RawHash) const;
+
+	FString CommonInputPath;
+	FString CommonOutputPath;
+	TArray<FString> ActionPaths;
+};
+
+static FSharedBuffer LoadFile(const FString& Path)
+{
+	FSharedBuffer Buffer;
+	if (TUniquePtr<FArchive> Ar{IFileManager::Get().CreateFileReader(*Path, FILEREAD_Silent)})
+	{
+		const int64 TotalSize = Ar->TotalSize();
+		FUniqueBuffer MutableBuffer = FUniqueBuffer::Alloc(uint64(TotalSize));
+		Ar->Serialize(MutableBuffer.GetData(), TotalSize);
+		if (Ar->Close())
+		{
+			Buffer = MutableBuffer.MoveToShared();
+		}
+	}
+	return Buffer;
+}
+
+bool FBuildWorkerProgram::ParseCommandLine(const TCHAR* CommandLine)
 {
-	FTaskTagScope Scope(ETaskTag::EGameThread);
-	FString CmdLine = FCommandLine::BuildFromArgV(nullptr, ArgC, ArgV, TEXT("-ddc=None"));
-	GEngineLoop.PreInit(*CmdLine);
+	TArray<FString> ActionPathPatterns;
+	TArray<FString> InputDirectoryPaths;
+	TArray<FString> OutputDirectoryPaths;
+
+	for (FString Token; FParse::Token(CommandLine, Token, /*UseEscape*/ false);)
+	{
+		Token.TrimQuotesInline();
+		const auto GetSwitchValues = [Token = FStringView(Token)](FStringView Match, TArray<FString>& OutValues)
+		{
+			if (Token.StartsWith(Match))
+			{
+				String::ParseTokens(Token.RightChop(Match.Len()), TEXT('+'),
+					[&OutValues](FStringView Value) { OutValues.Emplace(Value); });
+			}
+		};
+
+		GetSwitchValues(TEXT("-B="), ActionPathPatterns);
+		GetSwitchValues(TEXT("-Build="), ActionPathPatterns);
+
+		GetSwitchValues(TEXT("-I="), InputDirectoryPaths);
+		GetSwitchValues(TEXT("-Input="), InputDirectoryPaths);
+
+		GetSwitchValues(TEXT("-O="), OutputDirectoryPaths);
+		GetSwitchValues(TEXT("-Output="), OutputDirectoryPaths);
+	}
+
+	bool bCommandLineIsValid = true;
+
+	if (const int32 InputDirectoryCount = InputDirectoryPaths.Num(); InputDirectoryCount > 1)
+	{
+		UE_LOG(LogDerivedDataBuildWorker, Error, TEXT("A maximum of one input directory can be specified, but %d were specified."), InputDirectoryCount);
+		bCommandLineIsValid = false;
+	}
+	else if (InputDirectoryCount == 1)
+	{
+		CommonInputPath = FPaths::ConvertRelativePathToFull(FPaths::LaunchDir(), InputDirectoryPaths[0]);
+	}
+
+	if (const int32 OutputDirectoryCount = OutputDirectoryPaths.Num(); OutputDirectoryCount > 1)
+	{
+		UE_LOG(LogDerivedDataBuildWorker, Error, TEXT("A maximum of one input directory can be specified, but %d were specified."), OutputDirectoryCount);
+		bCommandLineIsValid = false;
+	}
+	else if (OutputDirectoryCount == 1)
+	{
+		CommonOutputPath = FPaths::ConvertRelativePathToFull(FPaths::LaunchDir(), OutputDirectoryPaths[0]);
+	}
+
+	if (ActionPathPatterns.IsEmpty())
+	{
+		UE_LOG(LogDerivedDataBuildWorker, Error, TEXT("No build action files specified on the command line."));
+		bCommandLineIsValid = false;
+	}
+
+	for (const FString& ActionPathPattern : ActionPathPatterns)
+	{
+		FWildcardString ActionPathWildcard(ActionPathPattern);
+		if (ActionPathWildcard.ContainsWildcards())
+		{
+			UE_LOG(LogDerivedDataBuildWorker, Error, TEXT("Wildcards in build action file paths are not supported yet: '%s'"), *ActionPathWildcard);
+			bCommandLineIsValid = false;
+		}
+		else
+		{
+			ActionPaths.Add(FPaths::ConvertRelativePathToFull(FPaths::LaunchDir(), ActionPathPattern));
+		}
+	}
+
+	return bCommandLineIsValid;
+}
+
+bool FBuildWorkerProgram::Build()
+{
+	IBuild& BuildSystem = GetDerivedDataBuildRef();
+	FBuildSession Session = BuildSystem.CreateSession(TEXT("BuildWorker"_SV), this);
+	TArray<FRequest> Builds;
 
-	// Make sure the engine is properly cleaned up whenever we exit this function
 	ON_SCOPE_EXIT
 	{
-		FEngineLoop::AppPreExit();
-		FEngineLoop::AppExit();
+		for (FRequest& Build : Builds)
+		{
+			Build.Wait();
+		}
 	};
 
-	UE::DerivedData::FBuildLoop BuildLoop;
-	if (!BuildLoop.Init())
+	for (const FString& ActionPath : ActionPaths)
 	{
-		return 1;
+		UE_LOG(LogDerivedDataBuildWorker, Log, TEXT("Loading build action '%s'"), *ActionPath);
+		FCbObject ActionObject;
+		if (TUniquePtr<FArchive> Ar{IFileManager::Get().CreateFileReader(*ActionPath, FILEREAD_Silent)})
+		{
+			*Ar << ActionObject;
+		}
+		else
+		{
+			UE_LOG(LogDerivedDataBuildWorker, Error, TEXT("Missing build action '%s'"), *ActionPath);
+			return false;
+		}
+		if (FOptionalBuildAction Action = BuildSystem.LoadAction(ActionPath, MoveTemp(ActionObject)); Action.IsNull())
+		{
+			UE_LOG(LogDerivedDataBuildWorker, Error, TEXT("Invalid build action '%s'"), *ActionPath);
+			return false;
+		}
+		else
+		{
+			Builds.Add(Session.BuildAction(Action.Get(), {},
+				EBuildPolicy::Local | EBuildPolicy::SkipCacheGet | EBuildPolicy::SkipCachePut, EPriority::Normal,
+				[this](FBuildActionCompleteParams&& Params) { BuildComplete(MoveTemp(Params)); }));
+		}
 	}
 
-	TMap<FName, const UE::DerivedData::IBuildFunction*> BuildFunctions;
-	for (UE::DerivedData::IBuildFunctionFactory* FunctionFactory : IModularFeatures::Get().GetModularFeatureImplementations<UE::DerivedData::IBuildFunctionFactory>(UE::DerivedData::IBuildFunctionFactory::GetFeatureName()))
+	if (Builds.IsEmpty())
 	{
-		const UE::DerivedData::IBuildFunction& BuildFunction = FunctionFactory->GetFunction();
-		BuildFunctions.FindOrAdd(FName(BuildFunction.GetName())) = &BuildFunction;
+		UE_LOG(LogDerivedDataBuildWorker, Error, TEXT("No build actions to operate on."));
+		return false;
 	}
 
-	BuildLoop.PerformBuilds([&] (FName FunctionName, UE::DerivedData::FBuildContext& BuildContext)
+	return true;
+}
+
+void FBuildWorkerProgram::BuildComplete(FBuildActionCompleteParams&& Params) const
+{
+	const FBuildOutput Output = MoveTemp(Params.Output);
+
+	if constexpr (!NO_LOGGING)
+	{
+		Output.IterateDiagnostics([](const FBuildDiagnostic& Diagnostic)
 		{
-			if (const UE::DerivedData::IBuildFunction** FoundFunc = BuildFunctions.Find(FunctionName))
+			ELogVerbosity::Type Verbosity;
+			switch (Diagnostic.Level)
 			{
-				const UE::DerivedData::IBuildFunction* BuildFunction = *FoundFunc;
-				UE_LOG(LogDerivedDataBuildWorker, Display, TEXT("Starting build function '%s'"), *FunctionName.ToString());
-				uint64 BuildStartTime = FPlatformTime::Cycles64();
-				BuildFunction->Build(BuildContext);
-				UE_LOG(LogDerivedDataBuildWorker, Display, TEXT("Completed in %fms"), FPlatformTime::ToMilliseconds64(FPlatformTime::Cycles64()-BuildStartTime));
-
-				return true;
+			default:
+			case EBuildDiagnosticLevel::Error:   Verbosity = ELogVerbosity::Error;   break;
+			case EBuildDiagnosticLevel::Warning: Verbosity = ELogVerbosity::Warning; break;
 			}
-			UE_LOG(LogDerivedDataBuildWorker, Error, TEXT("Unknown build function: %s"), *FunctionName.ToString());
-			return false;
+			GWarn->Log(FName(Diagnostic.Category), Verbosity, FString(Diagnostic.Message));
 		});
+	}
+
+	if (Output.HasError())
+	{
+		UE_LOG(LogDerivedDataBuildWorker, Error,
+			TEXT("Errors in build caused storage of payloads to be skipped for build of '%s' by %s."),
+			*WriteToString<128>(Output.GetName()), *WriteToString<32>(Output.GetFunction()));
+	}
+	else
+	{
+		for (const FPayload& Payload : Output.GetPayloads())
+		{
+			if (Payload.HasData())
+			{
+				if (TUniquePtr<FArchive> Ar = OpenOutput(Output.GetName(), Payload.GetRawHash()))
+				{
+					*Ar << const_cast<FCompressedBuffer&>(Payload.GetData());
+				}
+			}
+		}
+	}
 
-	BuildLoop.Teardown();
+	const FString OutputPath = FPaths::ChangeExtension(FString(Output.GetName()), TEXT("uddbo"));
+	if (TUniquePtr<FArchive> Ar{IFileManager::Get().CreateFileWriter(*OutputPath, FILEWRITE_Silent)})
+	{
+		FCbWriter OutputWriter;
+		Output.Save(OutputWriter);
+		OutputWriter.Save(*Ar);
+	}
+	else
+	{
+		UE_LOG(LogDerivedDataBuildWorker, Error,
+			TEXT("Failed to store build output to '%s' for build of '%s' by %s."),
+			*OutputPath, *WriteToString<128>(Output.GetName()), *WriteToString<32>(Output.GetFunction()));
+	}
+}
+
+FRequest FBuildWorkerProgram::ResolveInputData(const FBuildAction& Action, EPriority Priority, FOnBuildInputDataResolved&& OnResolved, FBuildInputFilter&& Filter)
+{
+	EStatus Status = EStatus::Ok;
+	TArray<FString> InputKeys;
+	TArray<FBuildInputDataByKey> Inputs;
+	Action.IterateInputs([this, &Action, &Filter, &InputKeys, &Inputs, &Status](FStringView Key, const FIoHash& RawHash, uint64 RawSize)
+	{
+		if (Filter && !Filter(Key))
+		{
+			return;
+		}
+		if (TUniquePtr<FArchive> Ar = OpenInput(Action.GetName(), RawHash))
+		{
+			InputKeys.Emplace(Key);
+			Inputs.Add({InputKeys.Last(), FCompressedBuffer::FromCompressed(*Ar)});
+		}
+		else
+		{
+			Status = EStatus::Error;
+			UE_LOG(LogDerivedDataBuildWorker, Error,
+				TEXT("Input %s with raw hash %s is missing for build of '%s' by %s."), *WriteToString<48>(RawHash),
+				*WriteToString<64>(Key), *WriteToString<128>(Action.GetName()), *WriteToString<32>(Action.GetFunction()));
+		}
+	});
+	OnResolved({Inputs, Status});
+	return FRequest();
+}
+
+TUniquePtr<FArchive> FBuildWorkerProgram::OpenInput(FStringView ActionPath, const FIoHash& RawHash) const
+{
+	TStringBuilder<256> Path;
+	if (CommonOutputPath.IsEmpty())
+	{
+		FPathViews::Append(Path, FPathViews::GetPath(ActionPath), TEXT("Inputs"), RawHash);
+	}
+	else
+	{
+		FPathViews::Append(Path, CommonOutputPath, RawHash);
+	}
+	return TUniquePtr<FArchive>(IFileManager::Get().CreateFileReader(*Path, FILEREAD_Silent));
+}
+
+TUniquePtr<FArchive> FBuildWorkerProgram::OpenOutput(FStringView ActionPath, const FIoHash& RawHash) const
+{
+	TStringBuilder<256> Path;
+	if (CommonOutputPath.IsEmpty())
+	{
+		FPathViews::Append(Path, FPathViews::GetPath(ActionPath), TEXT("Outputs"), RawHash);
+	}
+	else
+	{
+		FPathViews::Append(Path, CommonOutputPath, RawHash);
+	}
+	return TUniquePtr<FArchive>(IFileManager::Get().CreateFileWriter(*Path, FILEWRITE_NoReplaceExisting));
+}
+
+} // UE::DerivedData
+
+///////////////////////////////////////////////////////////////////////////////////////////////////////////////////////
+
+INT32_MAIN_INT32_ARGC_TCHAR_ARGV()
+{
+	ON_SCOPE_EXIT
+	{
+		FEngineLoop::AppPreExit();
+		FEngineLoop::AppExit();
+	};
+
+	const FTaskTagScope Scope(ETaskTag::EGameThread);
+
+	if (const int32 ErrorLevel = GEngineLoop.PreInit(ArgC, ArgV, TEXT("-DDC=None")))
+	{
+		return ErrorLevel;
+	}
+
+	UE::DerivedData::FBuildWorkerProgram Program;
+
+	if (!Program.ParseCommandLine(FCommandLine::Get()))
+	{
+		return 1;
+	}
+
+	if (!Program.Build())
+	{
+		return 1;
+	}
 
 	return 0;
 }
