DDC: Added a null check for GWarn in the build worker in case logging is disabled

#rnx

diff --git a/Engine/Source/Programs/DerivedDataBuildWorker/Private/DerivedDataBuildWorker.cpp b/Engine/Source/Programs/DerivedDataBuildWorker/Private/DerivedDataBuildWorker.cpp
--- a/Engine/Source/Programs/DerivedDataBuildWorker/Private/DerivedDataBuildWorker.cpp
+++ b/Engine/Source/Programs/DerivedDataBuildWorker/Private/DerivedDataBuildWorker.cpp
@@ -194,17 +194,20 @@ void FBuildWorkerProgram::BuildComplete(FBuildActionCompleteParams&& Params) con
 
 	if constexpr (!NO_LOGGING)
 	{
-		Output.IterateDiagnostics([](const FBuildDiagnostic& Diagnostic)
+		if (GWarn)
 		{
-			ELogVerbosity::Type Verbosity;
-			switch (Diagnostic.Level)
+			Output.IterateDiagnostics([](const FBuildDiagnostic& Diagnostic)
 			{
-			default:
-			case EBuildDiagnosticLevel::Error:   Verbosity = ELogVerbosity::Error;   break;
-			case EBuildDiagnosticLevel::Warning: Verbosity = ELogVerbosity::Warning; break;
-			}
-			GWarn->Log(FName(Diagnostic.Category), Verbosity, FString(Diagnostic.Message));
-		});
+				ELogVerbosity::Type Verbosity;
+				switch (Diagnostic.Level)
+				{
+				default:
+				case EBuildDiagnosticLevel::Error:   Verbosity = ELogVerbosity::Error;   break;
+				case EBuildDiagnosticLevel::Warning: Verbosity = ELogVerbosity::Warning; break;
+				}
+				GWarn->Log(FName(Diagnostic.Category), Verbosity, FString(Diagnostic.Message));
+			});
+		}
 	}
 
 	if (Output.HasError())
