DDC: Fixed the formatting of an error message

#rnx

diff --git a/Engine/Source/Programs/DerivedDataBuildWorker/Private/DerivedDataBuildWorker.cpp b/Engine/Source/Programs/DerivedDataBuildWorker/Private/DerivedDataBuildWorker.cpp
--- a/Engine/Source/Programs/DerivedDataBuildWorker/Private/DerivedDataBuildWorker.cpp
+++ b/Engine/Source/Programs/DerivedDataBuildWorker/Private/DerivedDataBuildWorker.cpp
@@ -265,8 +265,8 @@ FRequest FBuildWorkerProgram::ResolveInputData(const FBuildAction& Action, EPrio
 		{
 			Status = EStatus::Error;
 			UE_LOG(LogDerivedDataBuildWorker, Error,
-				TEXT("Input %s with raw hash %s is missing for build of '%s' by %s."), *WriteToString<48>(RawHash),
-				*WriteToString<64>(Key), *WriteToString<128>(Action.GetName()), *WriteToString<32>(Action.GetFunction()));
+				TEXT("Input '%s' with raw hash %s is missing for build of '%s' by %s."), *WriteToString<64>(Key),
+				*WriteToString<48>(RawHash), *WriteToString<128>(Action.GetName()), *WriteToString<32>(Action.GetFunction()));
 		}
 	});
 	OnResolved({Inputs, Status});
