Allow build workers to report the function names and versions they contain.

diff --git a/Engine/Source/Programs/DerivedDataBuildWorker/Private/DerivedDataBuildWorker.cpp b/Engine/Source/Programs/DerivedDataBuildWorker/Private/DerivedDataBuildWorker.cpp
--- a/Engine/Source/Programs/DerivedDataBuildWorker/Private/DerivedDataBuildWorker.cpp
+++ b/Engine/Source/Programs/DerivedDataBuildWorker/Private/DerivedDataBuildWorker.cpp
@@ -5,6 +5,7 @@
 #include "DerivedDataBuild.h"
 #include "DerivedDataBuildAction.h"
 #include "DerivedDataBuildDefinition.h"
+#include "DerivedDataBuildFunctionRegistry.h"
 #include "DerivedDataBuildInputs.h"
 #include "DerivedDataBuildInputResolver.h"
 #include "DerivedDataBuildOutput.h"
@@ -37,6 +38,7 @@ class FBuildWorkerProgram : public IBuildInputResolver
 {
 public:
 	bool ParseCommandLine(const TCHAR* CommandLine);
+	bool ReportVersions();
 	bool Build();
 
 private:
@@ -50,6 +52,7 @@ private:
 	FString CommonInputPath;
 	FString CommonOutputPath;
 	TArray<FString> ActionPaths;
+	TArray<FString> VersionPaths;
 };
 
 static FSharedBuffer LoadFile(const FString& Path)
@@ -94,6 +97,9 @@ bool FBuildWorkerProgram::ParseCommandLine(const TCHAR* CommandLine)
 
 		GetSwitchValues(TEXT("-O="), OutputDirectoryPaths);
 		GetSwitchValues(TEXT("-Output="), OutputDirectoryPaths);
+
+		GetSwitchValues(TEXT("-V="), VersionPaths);
+		GetSwitchValues(TEXT("-Version="), VersionPaths);
 	}
 
 	bool bCommandLineIsValid = true;
@@ -118,9 +124,9 @@ bool FBuildWorkerProgram::ParseCommandLine(const TCHAR* CommandLine)
 		CommonOutputPath = FPaths::ConvertRelativePathToFull(FPaths::LaunchDir(), OutputDirectoryPaths[0]);
 	}
 
-	if (ActionPathPatterns.IsEmpty())
+	if (ActionPathPatterns.IsEmpty() && VersionPaths.IsEmpty())
 	{
-		UE_LOG(LogDerivedDataBuildWorker, Error, TEXT("No build action files specified on the command line."));
+		UE_LOG(LogDerivedDataBuildWorker, Error, TEXT("No build action files or version files specified on the command line."));
 		bCommandLineIsValid = false;
 	}
 
@@ -141,8 +147,60 @@ bool FBuildWorkerProgram::ParseCommandLine(const TCHAR* CommandLine)
 	return bCommandLineIsValid;
 }
 
+bool FBuildWorkerProgram::ReportVersions()
+{
+	if (VersionPaths.IsEmpty())
+	{
+		return true;
+	}
+
+	IBuild& BuildSystem = GetDerivedDataBuildRef();
+	FGuid BuildSystemVersion = BuildSystem.GetVersion();
+	IBuildFunctionRegistry& BuildFunctionRegistry = BuildSystem.GetFunctionRegistry();
+	TMap<FStringView, FGuid> Functions;
+	BuildFunctionRegistry.IterateFunctionVersions([&Functions](FStringView Function, const FGuid& Version)
+	{
+		Functions.Emplace(Function, Version);
+	});
+
+	FCbWriter Writer;
+	Writer.BeginObject();
+
+	Writer << "BuildSystemVersion" << BuildSystemVersion;
+	UE_LOG(LogDerivedDataBuildWorker, Display, TEXT("BuildSystemVersion: '%s'"), *WriteToString<64>(BuildSystemVersion));
+	UE_LOG(LogDerivedDataBuildWorker, Display, TEXT("Functions:"));
+
+	Writer.BeginArray("Functions");
+	for (const TPair<FStringView, FGuid>& Function : Functions)
+	{
+		Writer.BeginObject();
+		Writer << "Name" << Function.Key;
+		Writer << "Version" << Function.Value;
+		Writer.EndObject();
+		UE_LOG(LogDerivedDataBuildWorker, Display, TEXT("%30.*s : '%s'"), Function.Key.Len(), Function.Key.GetData(), *WriteToString<64>(Function.Value));
+	}
+	Writer.EndArray();
+
+	Writer.EndObject();
+
+	for (const FString& VersionPath : VersionPaths)
+	{
+		if (TUniquePtr<FArchive> Ar{IFileManager::Get().CreateFileWriter(*FPaths::ConvertRelativePathToFull(FPaths::LaunchDir(), VersionPath))})
+		{
+			Writer.Save(*Ar);
+		}
+	}
+
+	return true;
+}
+
 bool FBuildWorkerProgram::Build()
 {
+	if (ActionPaths.IsEmpty())
+	{
+		return true;
+	}
+
 	IBuild& BuildSystem = GetDerivedDataBuildRef();
 	FBuildSession Session = BuildSystem.CreateSession(TEXT("BuildWorker"_SV), this);
 	TArray<FRequest> Builds;
@@ -317,6 +375,8 @@ INT32_MAIN_INT32_ARGC_TCHAR_ARGV()
 
 	const FTaskTagScope Scope(ETaskTag::EGameThread);
 
+	uint64 WorkerStartTime = FPlatformTime::Cycles64();
+
 	if (const int32 ErrorLevel = GEngineLoop.PreInit(ArgC, ArgV, TEXT("-DDC=None")))
 	{
 		return ErrorLevel;
@@ -329,10 +389,17 @@ INT32_MAIN_INT32_ARGC_TCHAR_ARGV()
 		return 1;
 	}
 
+	if (!Program.ReportVersions())
+	{
+		return 1;
+	}
+
 	if (!Program.Build())
 	{
 		return 1;
 	}
 
+	UE_LOG(LogDerivedDataBuildWorker, Display, TEXT("Worker completed in %fms"), FPlatformTime::ToMilliseconds64(FPlatformTime::Cycles64()-WorkerStartTime));
+
 	return 0;
 }
