Fix linux compile error and bug with my last submits related to function version iteration for build functions.

diff --git a/Engine/Source/Developer/DerivedDataCache/Private/DerivedDataBuildFunctionRegistry.cpp b/Engine/Source/Developer/DerivedDataCache/Private/DerivedDataBuildFunctionRegistry.cpp
--- a/Engine/Source/Developer/DerivedDataCache/Private/DerivedDataBuildFunctionRegistry.cpp
+++ b/Engine/Source/Developer/DerivedDataCache/Private/DerivedDataBuildFunctionRegistry.cpp
@@ -116,17 +116,17 @@ FGuid FBuildFunctionRegistry::FindFunctionVersion(FStringView Function) const
 
 void FBuildFunctionRegistry::IterateFunctionVersions(TFunctionRef<void(FStringView Function, const FGuid& Version)> Visitor) const
 {
-	TArray<TPair<FStringView, FGuid>> FunctionVersions;
+	TArray<TPair<FString, FGuid>> FunctionVersions;
 	{
 		FReadScopeLock ReadLock(Lock);
 		FunctionVersions.Reserve(Functions.Num());
-		for (const TPair<FStringView, IBuildFunctionFactory*>& Function : Functions)
+		for (const TPair<FString, IBuildFunctionFactory*>& Function : Functions)
 		{
 			FunctionVersions.Emplace(Function.Key, Function.Value->GetFunction().GetVersion());
 		}
 	}
 
-	for (const TPair<FStringView, FGuid>& Function : FunctionVersions)
+	for (const TPair<FString, FGuid>& Function : FunctionVersions)
 	{
 		Visitor(Function.Key, Function.Value);
 	}
diff --git a/Engine/Source/Programs/DerivedDataBuildWorker/Private/DerivedDataBuildWorker.cpp b/Engine/Source/Programs/DerivedDataBuildWorker/Private/DerivedDataBuildWorker.cpp
--- a/Engine/Source/Programs/DerivedDataBuildWorker/Private/DerivedDataBuildWorker.cpp
+++ b/Engine/Source/Programs/DerivedDataBuildWorker/Private/DerivedDataBuildWorker.cpp
@@ -157,7 +157,7 @@ bool FBuildWorkerProgram::ReportVersions()
 	IBuild& BuildSystem = GetDerivedDataBuildRef();
 	FGuid BuildSystemVersion = BuildSystem.GetVersion();
 	IBuildFunctionRegistry& BuildFunctionRegistry = BuildSystem.GetFunctionRegistry();
-	TMap<FStringView, FGuid> Functions;
+	TMap<FString, FGuid> Functions;
 	BuildFunctionRegistry.IterateFunctionVersions([&Functions](FStringView Function, const FGuid& Version)
 	{
 		Functions.Emplace(Function, Version);
@@ -171,13 +171,13 @@ bool FBuildWorkerProgram::ReportVersions()
 	UE_LOG(LogDerivedDataBuildWorker, Display, TEXT("Functions:"));
 
 	Writer.BeginArray("Functions");
-	for (const TPair<FStringView, FGuid>& Function : Functions)
+	for (const TPair<FString, FGuid>& Function : Functions)
 	{
 		Writer.BeginObject();
 		Writer << "Name" << Function.Key;
 		Writer << "Version" << Function.Value;
 		Writer.EndObject();
-		UE_LOG(LogDerivedDataBuildWorker, Display, TEXT("%30.*s : '%s'"), Function.Key.Len(), Function.Key.GetData(), *WriteToString<64>(Function.Value));
+		UE_LOG(LogDerivedDataBuildWorker, Display, TEXT("%30s : '%s'"), *Function.Key, *WriteToString<64>(Function.Value));
 	}
 	Writer.EndArray();
 
