Force load of all linked modules during initialization of a standalone build worker process.

diff --git a/Engine/Source/Programs/DerivedDataBuildWorker/Private/DerivedDataBuildWorker.cpp b/Engine/Source/Programs/DerivedDataBuildWorker/Private/DerivedDataBuildWorker.cpp
--- a/Engine/Source/Programs/DerivedDataBuildWorker/Private/DerivedDataBuildWorker.cpp
+++ b/Engine/Source/Programs/DerivedDataBuildWorker/Private/DerivedDataBuildWorker.cpp
@@ -6,8 +6,8 @@
 #include "DerivedDataBuildAction.h"
 #include "DerivedDataBuildDefinition.h"
 #include "DerivedDataBuildFunctionRegistry.h"
-#include "DerivedDataBuildInputs.h"
 #include "DerivedDataBuildInputResolver.h"
+#include "DerivedDataBuildInputs.h"
 #include "DerivedDataBuildOutput.h"
 #include "DerivedDataBuildSession.h"
 #include "DerivedDataPayload.h"
@@ -20,6 +20,7 @@
 #include "Misc/PathViews.h"
 #include "Misc/ScopeExit.h"
 #include "Misc/WildcardString.h"
+#include "Modules/ModuleManager.h"
 #include "Serialization/CompactBinary.h"
 #include "Serialization/CompactBinarySerialization.h"
 #include "Serialization/CompactBinaryWriter.h"
@@ -388,6 +389,14 @@ INT32_MAIN_INT32_ARGC_TCHAR_ARGV()
 		return 1;
 	}
 
+	FModuleManager& ModuleManager = FModuleManager::Get();
+	TArray<FName> Modules;
+	ModuleManager.FindModules(TEXT("*"), Modules);
+	for (FName Module : Modules)
+	{
+		ModuleManager.LoadModule(Module);
+	}
+
 	if (!Program.ReportVersions())
 	{
 		return 1;
