DDC: Fixed a compile error in the build worker

#rnx

diff --git a/Engine/Source/Programs/DerivedDataBuildWorker/Private/DerivedDataBuildWorker.cpp b/Engine/Source/Programs/DerivedDataBuildWorker/Private/DerivedDataBuildWorker.cpp
--- a/Engine/Source/Programs/DerivedDataBuildWorker/Private/DerivedDataBuildWorker.cpp
+++ b/Engine/Source/Programs/DerivedDataBuildWorker/Private/DerivedDataBuildWorker.cpp
@@ -155,7 +155,7 @@ bool FBuildWorkerProgram::ReportVersions()
 		return true;
 	}
 
-	IBuild& BuildSystem = GetDerivedDataBuildRef();
+	IBuild& BuildSystem = GetBuild();
 	FGuid BuildSystemVersion = BuildSystem.GetVersion();
 	IBuildFunctionRegistry& BuildFunctionRegistry = BuildSystem.GetFunctionRegistry();
 	TMap<FString, FGuid> Functions;
@@ -202,7 +202,7 @@ bool FBuildWorkerProgram::Build()
 		return true;
 	}
 
-	IBuild& BuildSystem = GetDerivedDataBuildRef();
+	IBuild& BuildSystem = GetBuild();
 	FBuildSession Session = BuildSystem.CreateSession(TEXT("BuildWorker"_SV), this);
 	FRequestGroup Group = BuildSystem.CreateGroup(EPriority::Normal);
 
