Avoid parsing '+' characters as separators of a multi-value commandline argument for DerivedDataBuildWorkers as it causes bugs when a simple path has a '+' character in it as occurs on the build farm.

#rnx

diff --git a/Engine/Source/Programs/DerivedDataBuildWorker/Private/DerivedDataBuildWorker.cpp b/Engine/Source/Programs/DerivedDataBuildWorker/Private/DerivedDataBuildWorker.cpp
--- a/Engine/Source/Programs/DerivedDataBuildWorker/Private/DerivedDataBuildWorker.cpp
+++ b/Engine/Source/Programs/DerivedDataBuildWorker/Private/DerivedDataBuildWorker.cpp
@@ -86,8 +86,7 @@ bool FBuildWorkerProgram::ParseCommandLine(const TCHAR* CommandLine)
 		{
 			if (Token.StartsWith(Match))
 			{
-				String::ParseTokens(Token.RightChop(Match.Len()), TEXT('+'),
-					[&OutValues](FStringView Value) { OutValues.Emplace(Value); });
+				OutValues.Emplace(Token.RightChop(Match.Len()));
 			}
 		};
 
