DDC: Check for existence of inputs in the build worker

Cancel unfinished builds when an error occurs that will lead to a non-zero exit code, and fix handling of -I/-Input.

#rnx

diff --git a/Engine/Source/Programs/DerivedDataBuildWorker/Private/DerivedDataBuildWorker.cpp b/Engine/Source/Programs/DerivedDataBuildWorker/Private/DerivedDataBuildWorker.cpp
--- a/Engine/Source/Programs/DerivedDataBuildWorker/Private/DerivedDataBuildWorker.cpp
+++ b/Engine/Source/Programs/DerivedDataBuildWorker/Private/DerivedDataBuildWorker.cpp
@@ -46,8 +46,12 @@ public:
 private:
 	void BuildComplete(FBuildActionCompleteParams&& Params) const;
 
+	bool ResolveInputExists(const FBuildAction& Action) const;
 	void ResolveInputData(const FBuildAction& Action, IRequestOwner& Owner, FOnBuildInputDataResolved&& OnResolved, FBuildInputFilter&& Filter) final;
 
+	void GetInputPath(FStringView ActionPath, const FIoHash& RawHash, FStringBuilderBase& OutPath) const;
+	void GetOutputPath(FStringView ActionPath, const FIoHash& RawHash, FStringBuilderBase& OutPath) const;
+
 	TUniquePtr<FArchive> OpenInput(FStringView ActionPath, const FIoHash& RawHash) const;
 	TUniquePtr<FArchive> OpenOutput(FStringView ActionPath, const FIoHash& RawHash) const;
 
@@ -205,18 +209,6 @@ bool FBuildWorkerProgram::Build()
 	IBuild& BuildSystem = GetBuild();
 	FBuildSession Session = BuildSystem.CreateSession(TEXT("BuildWorker"_SV), this);
 	FRequestOwner Owner(EPriority::Normal);
-
-	ON_SCOPE_EXIT
-	{
-		Owner.Wait();
-	};
-
-	if (ActionPaths.IsEmpty())
-	{
-		UE_LOG(LogDerivedDataBuildWorker, Error, TEXT("No build actions to operate on."));
-		return false;
-	}
-	else
 	{
 		FRequestBarrier Barrier(Owner);
 		for (const FString& ActionPath : ActionPaths)
@@ -237,6 +229,10 @@ bool FBuildWorkerProgram::Build()
 				UE_LOG(LogDerivedDataBuildWorker, Error, TEXT("Invalid build action '%s'"), *ActionPath);
 				return false;
 			}
+			else if (!ResolveInputExists(Action.Get()))
+			{
+				return false;
+			}
 			else
 			{
 				Session.BuildAction(Action.Get(), {}, EBuildPolicy::BuildLocal, Owner,
@@ -244,7 +240,7 @@ bool FBuildWorkerProgram::Build()
 			}
 		}
 	}
-
+	Owner.Wait();
 	return true;
 }
 
@@ -305,6 +301,24 @@ void FBuildWorkerProgram::BuildComplete(FBuildActionCompleteParams&& Params) con
 	}
 }
 
+bool FBuildWorkerProgram::ResolveInputExists(const FBuildAction& Action) const
+{
+	bool bValid = true;
+	Action.IterateInputs([this, &Action, &bValid](FStringView Key, const FIoHash& RawHash, uint64 RawSize)
+	{
+		TStringBuilder<256> Path;
+		GetInputPath(Action.GetName(), RawHash, Path);
+		if (!IFileManager::Get().FileExists(*Path))
+		{
+			bValid = false;
+			UE_LOG(LogDerivedDataBuildWorker, Error,
+				TEXT("Input '%s' with raw hash %s is missing for build of '%s' by %s."), *WriteToString<64>(Key),
+				*WriteToString<48>(RawHash), *WriteToString<128>(Action.GetName()), *WriteToString<32>(Action.GetFunction()));
+		}
+	});
+	return bValid;
+}
+
 void FBuildWorkerProgram::ResolveInputData(const FBuildAction& Action, IRequestOwner& Owner, FOnBuildInputDataResolved&& OnResolved, FBuildInputFilter&& Filter)
 {
 	EStatus Status = EStatus::Ok;
@@ -332,31 +346,41 @@ void FBuildWorkerProgram::ResolveInputData(const FBuildAction& Action, IRequestO
 	OnResolved({Inputs, Status});
 }
 
-TUniquePtr<FArchive> FBuildWorkerProgram::OpenInput(FStringView ActionPath, const FIoHash& RawHash) const
+void FBuildWorkerProgram::GetInputPath(FStringView ActionPath, const FIoHash& RawHash, FStringBuilderBase& OutPath) const
 {
-	TStringBuilder<256> Path;
-	if (CommonOutputPath.IsEmpty())
+	if (CommonInputPath.IsEmpty())
 	{
-		FPathViews::Append(Path, FPathViews::GetPath(ActionPath), TEXT("Inputs"), RawHash);
+		FPathViews::Append(OutPath, FPathViews::GetPath(ActionPath), TEXT("Inputs"), RawHash);
 	}
 	else
 	{
-		FPathViews::Append(Path, CommonOutputPath, RawHash);
+		FPathViews::Append(OutPath, CommonInputPath, RawHash);
 	}
-	return TUniquePtr<FArchive>(IFileManager::Get().CreateFileReader(*Path, FILEREAD_Silent));
 }
 
-TUniquePtr<FArchive> FBuildWorkerProgram::OpenOutput(FStringView ActionPath, const FIoHash& RawHash) const
+void FBuildWorkerProgram::GetOutputPath(FStringView ActionPath, const FIoHash& RawHash, FStringBuilderBase& OutPath) const
 {
-	TStringBuilder<256> Path;
 	if (CommonOutputPath.IsEmpty())
 	{
-		FPathViews::Append(Path, FPathViews::GetPath(ActionPath), TEXT("Outputs"), RawHash);
+		FPathViews::Append(OutPath, FPathViews::GetPath(ActionPath), TEXT("Outputs"), RawHash);
 	}
 	else
 	{
-		FPathViews::Append(Path, CommonOutputPath, RawHash);
+		FPathViews::Append(OutPath, CommonOutputPath, RawHash);
 	}
+}
+
+TUniquePtr<FArchive> FBuildWorkerProgram::OpenInput(FStringView ActionPath, const FIoHash& RawHash) const
+{
+	TStringBuilder<256> Path;
+	GetInputPath(ActionPath, RawHash, Path);
+	return TUniquePtr<FArchive>(IFileManager::Get().CreateFileReader(*Path, FILEREAD_Silent));
+}
+
+TUniquePtr<FArchive> FBuildWorkerProgram::OpenOutput(FStringView ActionPath, const FIoHash& RawHash) const
+{
+	TStringBuilder<256> Path;
+	GetOutputPath(ActionPath, RawHash, Path);
 	return TUniquePtr<FArchive>(IFileManager::Get().CreateFileWriter(*Path, FILEWRITE_NoReplaceExisting));
 }
 
