DDC: Updated DerivedDataBuildWorker to use the new build session interface

#rnx

diff --git a/Engine/Source/Programs/DerivedDataBuildWorker/Private/DerivedDataBuildWorker.cpp b/Engine/Source/Programs/DerivedDataBuildWorker/Private/DerivedDataBuildWorker.cpp
--- a/Engine/Source/Programs/DerivedDataBuildWorker/Private/DerivedDataBuildWorker.cpp
+++ b/Engine/Source/Programs/DerivedDataBuildWorker/Private/DerivedDataBuildWorker.cpp
@@ -44,7 +44,7 @@ public:
 	bool Build();
 
 private:
-	void BuildComplete(FBuildActionCompleteParams&& Params) const;
+	void BuildComplete(FBuildCompleteParams&& Params) const;
 
 	bool ResolveInputExists(const FBuildAction& Action) const;
 	void ResolveInputData(const FBuildAction& Action, IRequestOwner& Owner, FOnBuildInputDataResolved&& OnResolved, FBuildInputFilter&& Filter) final;
@@ -239,8 +239,8 @@ bool FBuildWorkerProgram::Build()
 			}
 			else
 			{
-				Session.BuildAction(Action.Get(), {}, EBuildPolicy::BuildLocal, Owner,
-					[this](FBuildActionCompleteParams&& Params) { BuildComplete(MoveTemp(Params)); });
+				Session.Build(Action.Get(), {}, EBuildPolicy::BuildLocal, Owner,
+					[this](FBuildCompleteParams&& Params) { BuildComplete(MoveTemp(Params)); });
 			}
 		}
 	}
@@ -248,7 +248,7 @@ bool FBuildWorkerProgram::Build()
 	return true;
 }
 
-void FBuildWorkerProgram::BuildComplete(FBuildActionCompleteParams&& Params) const
+void FBuildWorkerProgram::BuildComplete(FBuildCompleteParams&& Params) const
 {
 	const FBuildOutput Output = MoveTemp(Params.Output);
 
