DDC: Fixed a compile error in the derived data build worker

#rnx

diff --git a/Engine/Source/Programs/DerivedDataBuildWorker/Private/DerivedDataBuildWorker.cpp b/Engine/Source/Programs/DerivedDataBuildWorker/Private/DerivedDataBuildWorker.cpp
--- a/Engine/Source/Programs/DerivedDataBuildWorker/Private/DerivedDataBuildWorker.cpp
+++ b/Engine/Source/Programs/DerivedDataBuildWorker/Private/DerivedDataBuildWorker.cpp
@@ -364,7 +364,7 @@ void FBuildWorkerProgram::ResolveInputData(const FBuildAction& Action, IRequestO
 		if (TUniquePtr<FArchive> Ar = OpenInput(Action.GetName(), RawHash))
 		{
 			InputKeys.Emplace(Key);
-			Inputs.Add({InputKeys.Last(), FCompressedBuffer::FromCompressed(*Ar)});
+			Inputs.Add({InputKeys.Last(), FCompressedBuffer::Load(*Ar)});
 		}
 		else
 		{
