TextureBuildWorker: Force set -unattended when building

#rnx

diff --git a/Engine/Source/Programs/DerivedDataBuildWorker/Private/DerivedDataBuildWorker.cpp b/Engine/Source/Programs/DerivedDataBuildWorker/Private/DerivedDataBuildWorker.cpp
--- a/Engine/Source/Programs/DerivedDataBuildWorker/Private/DerivedDataBuildWorker.cpp
+++ b/Engine/Source/Programs/DerivedDataBuildWorker/Private/DerivedDataBuildWorker.cpp
@@ -431,7 +431,7 @@ INT32_MAIN_INT32_ARGC_TCHAR_ARGV()
 
 	uint64 WorkerStartTime = FPlatformTime::Cycles64();
 
-	if (const int32 ErrorLevel = GEngineLoop.PreInit(ArgC, ArgV, TEXT("-DDC=None")))
+	if (const int32 ErrorLevel = GEngineLoop.PreInit(ArgC, ArgV, TEXT("-DDC=None -Unattended")))
 	{
 		return ErrorLevel;
 	}
