DerivedDataBuildWorker: RequestEngineExit() when terminating

#rnx

diff --git a/Engine/Source/Programs/DerivedDataBuildWorker/Private/DerivedDataBuildWorker.cpp b/Engine/Source/Programs/DerivedDataBuildWorker/Private/DerivedDataBuildWorker.cpp
--- a/Engine/Source/Programs/DerivedDataBuildWorker/Private/DerivedDataBuildWorker.cpp
+++ b/Engine/Source/Programs/DerivedDataBuildWorker/Private/DerivedDataBuildWorker.cpp
@@ -422,7 +422,9 @@ INT32_MAIN_INT32_ARGC_TCHAR_ARGV()
 {
 	ON_SCOPE_EXIT
 	{
+		RequestEngineExit(TEXT("Exiting"));
 		FEngineLoop::AppPreExit();
+		FModuleManager::Get().UnloadModulesAtShutdown();
 		FEngineLoop::AppExit();
 	};
 
