DDC: Changed the build worker shutdown flow to more closely approximate the editor

#rnx

diff --git a/Engine/Source/Programs/DerivedDataBuildWorker/Private/DerivedDataBuildWorker.cpp b/Engine/Source/Programs/DerivedDataBuildWorker/Private/DerivedDataBuildWorker.cpp
--- a/Engine/Source/Programs/DerivedDataBuildWorker/Private/DerivedDataBuildWorker.cpp
+++ b/Engine/Source/Programs/DerivedDataBuildWorker/Private/DerivedDataBuildWorker.cpp
@@ -445,6 +445,11 @@ INT32_MAIN_INT32_ARGC_TCHAR_ARGV()
 	}
 
 	FModuleManager& ModuleManager = FModuleManager::Get();
+
+	// Load DerivedDataCache before the rest of the modules because that will make it the last to
+	// shut down on exit. This is an approximation of the module load/unload order in the editor.
+	ModuleManager.LoadModule(TEXT("DerivedDataCache"));
+
 	TArray<FName> Modules;
 	ModuleManager.FindModules(TEXT("*"), Modules);
 	for (FName Module : Modules)
