DDC: Follow-up to allowing the build system to be created without a cache

This change was included when testing but missing from the submitted change.

diff --git a/Engine/Source/Programs/DerivedDataBuildWorker/Private/DerivedDataBuildWorker.cpp b/Engine/Source/Programs/DerivedDataBuildWorker/Private/DerivedDataBuildWorker.cpp
--- a/Engine/Source/Programs/DerivedDataBuildWorker/Private/DerivedDataBuildWorker.cpp
+++ b/Engine/Source/Programs/DerivedDataBuildWorker/Private/DerivedDataBuildWorker.cpp
@@ -432,7 +432,7 @@ INT32_MAIN_INT32_ARGC_TCHAR_ARGV()
 
 	uint64 WorkerStartTime = FPlatformTime::Cycles64();
 
-	if (const int32 ErrorLevel = GEngineLoop.PreInit(ArgC, ArgV, TEXT("-DDC=None -Unattended")))
+	if (const int32 ErrorLevel = GEngineLoop.PreInit(ArgC, ArgV, TEXT("-Unattended")))
 	{
 		return ErrorLevel;
 	}
