DDB: Updates to DerivedDataBuildWorker

- Updated target receipt processing to add files from ProducedItems which are needed to support modular targets.
- Updated default DerivedDataBuildWorker TargetRules to further reduce code size.
- Exposed the build worker entry point to support modular targets.

#Synced-CL 49858055

diff --git a/Engine/Source/Developer/DesktopPlatform/Private/TargetReceiptBuildWorker.cpp b/Engine/Source/Developer/DesktopPlatform/Private/TargetReceiptBuildWorker.cpp
--- a/Engine/Source/Developer/DesktopPlatform/Private/TargetReceiptBuildWorker.cpp
+++ b/Engine/Source/Developer/DesktopPlatform/Private/TargetReceiptBuildWorker.cpp
@@ -2,6 +2,8 @@
 
 #include "TargetReceiptBuildWorker.h"
 
+#include "Algo/Sort.h"
+#include "Algo/Unique.h"
 #include "Compression/CompressedBuffer.h"
 #include "Compression/OodleDataCompression.h"
 #include "DerivedDataBuildWorker.h"
@@ -352,6 +354,20 @@ void FTargetReceiptBuildWorker::PopulateWorkerFromReceipt(const FString& TargetR
 		return;
 	}
 
+	for (const FBuildProduct& BuildProduct : TargetReceipt.BuildProducts)
+	{
+		if (BuildProduct.Type == TEXTVIEW("SymbolFile") ||
+			BuildProduct.Path == TargetReceipt.Launch)
+		{
+			continue;
+		}
+		if (!TryAddExecutablePath(BuildProduct.Path))
+		{
+			UE_LOG(LogTargetReceiptBuildWorker, Warning, TEXT("Target receipt referenced a build product that was not usable: '%s'"), *BuildProduct.Path);
+			return;
+		}
+	}
+
 	for (const FRuntimeDependency& RuntimeDependency : TargetReceipt.RuntimeDependencies)
 	{
 		if (RuntimeDependency.Path.EndsWith(TEXT(".uproject")))
@@ -371,6 +387,12 @@ void FTargetReceiptBuildWorker::PopulateWorkerFromReceipt(const FString& TargetR
 		return;
 	}
 
+	// Build a sorted list of unique paths because the target receipt may reference the same file in multiple places.
+	// Skip the first path because it is assumed elsewhere to be the launch path.
+	Algo::SortBy(MakeArrayView(ExecutablePaths).RightChop(1), &FWorkerPath::GetLocalPath);
+	const int32 KeepPaths = Algo::UniqueBy(ExecutablePaths, &FWorkerPath::GetLocalPath);
+	ExecutablePaths.RemoveAt(KeepPaths, ExecutablePaths.Num() - KeepPaths, EAllowShrinking::No);
+
 	// Open executable files to check for existence and gather hash info for them
 	int32 MetaIndex = 0;
 	ExecutableMeta.AddDefaulted(ExecutablePaths.Num());
diff --git a/Engine/Source/Programs/DerivedDataBuildWorker/DerivedDataBuildWorker.Build.cs b/Engine/Source/Programs/DerivedDataBuildWorker/DerivedDataBuildWorker.Build.cs
--- a/Engine/Source/Programs/DerivedDataBuildWorker/DerivedDataBuildWorker.Build.cs
+++ b/Engine/Source/Programs/DerivedDataBuildWorker/DerivedDataBuildWorker.Build.cs
@@ -13,37 +13,43 @@ public abstract class DerivedDataBuildWorkerTarget : TargetRules
 
 		SolutionDirectory = "Programs/BuildWorker";
 
-		bUseXGEController				= false;
-		bCompileFreeType				= false;
-		bLoggingToMemoryEnabled			= true;
-		bUseLoggingInShipping			= true;
-		bCompileWithAccessibilitySupport= false;
-		bWithServerCode					= false;
-		bCompileNavmeshClusterLinks		= false;
-		bCompileNavmeshSegmentLinks		= false;
-		bCompileRecast					= false;
-		bCompileICU 					= false;
-		bWithLiveCoding					= false;
-		bBuildDeveloperTools			= false;
-		bBuildWithEditorOnlyData		= true;
-		bCompileAgainstEngine			= false;
-		bCompileAgainstCoreUObject		= false;
-		bCompileAgainstApplicationCore	= false;
-		bUsesSlate = false;
-		bIsBuildingConsoleApplication	= true;
+		bLoggingToMemoryEnabled				= true;
+		bUseLoggingInShipping				= true;
+
+		bCompileCEF3						= false;
+		bCompileFreeType					= false;
+		bCompileICU 						= false;
+		bCompileWithAccessibilitySupport	= false;
+		bCompileWithPluginSupport			= false;
+		bUsesSlate							= false;
+		bUseXGEController					= false;
+		bWithLiveCoding						= false;
+		bWithServerCode						= false;
+
+		bForceDisableAutomationTests		= true;
+
+		bIsBuildingConsoleApplication		= true;
+		bBuildWithEditorOnlyData			= true;
+		bBuildDeveloperTools				= false;
+		bCompileAgainstEngine				= false;
+		bCompileAgainstCoreUObject			= false;
+		bCompileAgainstApplicationCore		= false;
 		// TODO: I want to use static CRT in the future, but it causes link issues today likely due to 3rd party libraries
-		//bUseStaticCRT					= true;
+		//bUseStaticCRT						= true;
 
 		WindowsPlatform.bUseBundledDbgHelp = false;
 		WindowsPlatform.bPixProfilingEnabled = false;
 		WindowsPlatform.bMergeIdenticalCOMDATs = true;
 		WindowsPlatform.bStripUnreferencedSymbols = true;
 
-		// Disable logging, as the workers are spawned often and logging will just slow them down
+		// Disable logging to a file to improve worker performance.
 		GlobalDefinitions.Add("ALLOW_LOG_FILE=0");
-		// Linking against wer.lib/wer.dll causes XGE to bail when the worker is run on a Windows 8 machine, so turn this off.
+		// Linking against wer.lib/wer.dll causes XGE to bail when the worker is run on a Windows 8 machine.
+		// TODO: Is this necessary with a Windows 10 minimum and past the end of support for Windows 10?
 		GlobalDefinitions.Add("ALLOW_WINDOWS_ERROR_REPORT_LIB=0");
-		// Disable external profiling in TextureBuilder to improve startup time
+		// Disable initialization of the crash reporter, which removes another thread being created.
+		GlobalDefinitions.Add("NOINITCRASHREPORTER=1");
+		// Disable external profiling to improve startup time.
 		GlobalDefinitions.Add("UE_EXTERNAL_PROFILING_ENABLED=0");
 	}
 }
@@ -67,7 +73,7 @@ public class DerivedDataBuildWorker : ModuleRules
 			PrivateDependencyModuleNames.Add("ApplicationCore");
 		}
 
-		// This needs to match the version defined in TargetReceiptBuildWorkerFactory.cpp
+		// This needs to match the version defined in TargetReceiptBuildWorker.cpp
 		AdditionalPropertiesForReceipt.Add("DerivedDataBuildWorkerReceiptVersion", "dab5352e-a5a7-4793-a7a3-1d4acad6aff2");
 	}
 }
diff --git a/Engine/Source/Programs/DerivedDataBuildWorker/Private/DerivedDataBuildWorker.cpp b/Engine/Source/Programs/DerivedDataBuildWorker/Private/DerivedDataBuildWorker.cpp
--- a/Engine/Source/Programs/DerivedDataBuildWorker/Private/DerivedDataBuildWorker.cpp
+++ b/Engine/Source/Programs/DerivedDataBuildWorker/Private/DerivedDataBuildWorker.cpp
@@ -21,6 +21,7 @@
 #include "Misc/Paths.h"
 #include "Misc/PathViews.h"
 #include "Misc/ScopeExit.h"
+#include "Misc/StringBuilder.h"
 #include "Misc/WildcardString.h"
 #include "Modules/ModuleManager.h"
 #include "Serialization/CompactBinary.h"
@@ -418,7 +419,7 @@ TUniquePtr<FArchive> FBuildWorkerProgram::OpenOutput(FStringView ActionPath, con
 
 ///////////////////////////////////////////////////////////////////////////////////////////////////////////////////////
 
-INT32_MAIN_INT32_ARGC_TCHAR_ARGV()
+DERIVEDDATABUILDWORKER_API int32 BuildWorkerMain(int32 ArgC, TCHAR* ArgV[], const TCHAR* ExtraArgs, const TCHAR* LoadModules)
 {
 	ON_SCOPE_EXIT
 	{
@@ -432,7 +433,7 @@ INT32_MAIN_INT32_ARGC_TCHAR_ARGV()
 
 	uint64 WorkerStartTime = FPlatformTime::Cycles64();
 
-	if (const int32 ErrorLevel = GEngineLoop.PreInit(ArgC, ArgV, TEXT("-Unattended")))
+	if (const int32 ErrorLevel = GEngineLoop.PreInit(ArgC, ArgV, ExtraArgs))
 	{
 		return ErrorLevel;
 	}
@@ -451,11 +452,15 @@ INT32_MAIN_INT32_ARGC_TCHAR_ARGV()
 	ModuleManager.LoadModule(TEXT("DerivedDataCache"));
 
 	TArray<FName> Modules;
-	ModuleManager.FindModules(TEXT("*"), Modules);
-	for (FName Module : Modules)
+	UE::String::ParseTokens(LoadModules, '+', [&ModuleManager](FStringView ModulePattern)
 	{
-		ModuleManager.LoadModule(Module);
-	}
+		TArray<FName> Modules;
+		ModuleManager.FindModules(*WriteToString<64>(ModulePattern), Modules);
+		for (FName Module : Modules)
+		{
+			ModuleManager.LoadModule(Module);
+		}
+	});
 
 	if (!Program.ReportVersions())
 	{
@@ -471,3 +476,8 @@ INT32_MAIN_INT32_ARGC_TCHAR_ARGV()
 
 	return 0;
 }
+
+INT32_MAIN_INT32_ARGC_TCHAR_ARGV()
+{
+	return BuildWorkerMain(ArgC, ArgV, TEXT("-Unattended"), TEXT("*"));
+}
